Editor: Fix depth picker to actually pick a sensible depth rather than just the hit proxy actor location

diff --git a/Engine/Source/Editor/SceneDepthPickerMode/Private/EditorModeSceneDepthPicker.cpp b/Engine/Source/Editor/SceneDepthPickerMode/Private/EditorModeSceneDepthPicker.cpp
--- a/Engine/Source/Editor/SceneDepthPickerMode/Private/EditorModeSceneDepthPicker.cpp
+++ b/Engine/Source/Editor/SceneDepthPickerMode/Private/EditorModeSceneDepthPicker.cpp
@@ -5,6 +5,8 @@
 #include "Widgets/SToolTip.h"
 #include "LevelEditorViewport.h"
 #include "EditorModes.h"
+#include "Editor/ActorPositioning.h"
+
 
 #define LOCTEXT_NAMESPACE "SceneDepthPicker"
 
@@ -77,13 +79,22 @@ bool FEdModeSceneDepthPicker::InputKey(FEditorViewportClient* ViewportClient, FV
 	{
 		if (Key == EKeys::LeftMouseButton && Event == IE_Pressed)
 		{
-			// See if we clicked on an actor
 			int32 const HitX = Viewport->GetMouseX();
 			int32 const HitY = Viewport->GetMouseY();
 
-			FVector const ObjectLoc = ViewportClient->GetHitProxyObjectLocation(HitX, HitY);
+			FSceneViewFamilyContext ViewFamily(FSceneViewFamily::ConstructionValues(
+				Viewport,
+				ViewportClient->GetScene(),
+				ViewportClient->EngineShowFlags)
+				.SetRealtimeUpdate(true));
+			FSceneView* View = ViewportClient->CalcSceneView(&ViewFamily);
+			FViewportCursorLocation Cursor(View, ViewportClient, HitX, HitY);
+
+			const FActorPositionTraceResult TraceResult = FActorPositioning::TraceWorldForPosition(Cursor, *View);
+
+			FVector WorldPosition = TraceResult.State == FActorPositionTraceResult::HitSuccess ? TraceResult.Location : ViewportClient->GetHitProxyObjectLocation(HitX, HitY);
 
-			OnSceneDepthLocationSelected.ExecuteIfBound(ObjectLoc);
+			OnSceneDepthLocationSelected.ExecuteIfBound(WorldPosition);
 			RequestDeletion();
 			return true;
 		}
diff --git a/Engine/Source/Editor/SceneDepthPickerMode/SceneDepthPickerMode.Build.cs b/Engine/Source/Editor/SceneDepthPickerMode/SceneDepthPickerMode.Build.cs
--- a/Engine/Source/Editor/SceneDepthPickerMode/SceneDepthPickerMode.Build.cs
+++ b/Engine/Source/Editor/SceneDepthPickerMode/SceneDepthPickerMode.Build.cs
@@ -6,7 +6,14 @@ public class SceneDepthPickerMode : ModuleRules
 {
     public SceneDepthPickerMode(ReadOnlyTargetRules Target) : base(Target)
 	{
-		PrivateDependencyModuleNames.AddRange(
+        PrivateIncludePaths.AddRange(
+            new string[]
+            {
+                "Editor/UnrealEd/Private"
+            }
+        );
+
+        PrivateDependencyModuleNames.AddRange(
 			new string[] {
 				"Core",
 				"CoreUObject",
